Added a dummy plugin for WebBrowser so that it can be included optionally

Removed dependency on WebBrowser from UE4Game

diff --git a/Engine/Plugins/Runtime/WebBrowserAsPlugin/Source/WebBrowserAsPlugin/Private/WebBrowserAsPluginModule.cpp b/Engine/Plugins/Runtime/WebBrowserAsPlugin/Source/WebBrowserAsPlugin/Private/WebBrowserAsPluginModule.cpp
new file mode 100644
--- /dev/null
+++ b/Engine/Plugins/Runtime/WebBrowserAsPlugin/Source/WebBrowserAsPlugin/Private/WebBrowserAsPluginModule.cpp
@@ -0,0 +1,23 @@
+// Copyright 1998-2015 Epic Games, Inc. All Rights Reserved.
+
+#include "WebBrowserAsPluginPrivatePCH.h"
+
+//////////////////////////////////////////////////////////////////////////
+// FWebBrowserAsPluginModule
+
+class FWebBrowserAsPluginModule : public IWebBrowserAsPluginModule
+{
+public:
+	virtual void StartupModule() override
+	{
+	}
+
+	virtual void ShutdownModule() override
+	{
+	}
+};
+
+//////////////////////////////////////////////////////////////////////////
+
+IMPLEMENT_MODULE(FWebBrowserAsPluginModule, WebBrowserAsPlugin);
+DEFINE_LOG_CATEGORY(LogWebBrowserAsPlugin);
diff --git a/Engine/Plugins/Runtime/WebBrowserAsPlugin/Source/WebBrowserAsPlugin/Private/WebBrowserAsPluginPrivatePCH.h b/Engine/Plugins/Runtime/WebBrowserAsPlugin/Source/WebBrowserAsPlugin/Private/WebBrowserAsPluginPrivatePCH.h
new file mode 100644
--- /dev/null
+++ b/Engine/Plugins/Runtime/WebBrowserAsPlugin/Source/WebBrowserAsPlugin/Private/WebBrowserAsPluginPrivatePCH.h
@@ -0,0 +1,5 @@
+// Copyright 1998-2015 Epic Games, Inc. All Rights Reserved.
+
+#include "WebBrowserAsPluginModule.h"
+
+DECLARE_LOG_CATEGORY_EXTERN(LogWebBrowserAsPlugin, Verbose, All);
\ No newline at end of file
diff --git a/Engine/Plugins/Runtime/WebBrowserAsPlugin/Source/WebBrowserAsPlugin/Public/WebBrowserAsPluginModule.h b/Engine/Plugins/Runtime/WebBrowserAsPlugin/Source/WebBrowserAsPlugin/Public/WebBrowserAsPluginModule.h
new file mode 100644
--- /dev/null
+++ b/Engine/Plugins/Runtime/WebBrowserAsPlugin/Source/WebBrowserAsPlugin/Public/WebBrowserAsPluginModule.h
@@ -0,0 +1,37 @@
+// Copyright 1998-2015 Epic Games, Inc. All Rights Reserved.
+
+#pragma once
+
+#include "ModuleManager.h"
+
+
+/**
+ * The public interface to this module
+ */
+class IWebBrowserAsPluginModule : public IModuleInterface
+{
+
+public:
+
+	/**
+	 * Singleton-like access to this module's interface.  This is just for convenience!
+	 * Beware of calling this during the shutdown phase, though.  Your module might have been unloaded already.
+	 *
+	 * @return Returns singleton instance, loading the module on demand if needed
+	 */
+	static inline IWebBrowserAsPluginModule& Get()
+	{
+		return FModuleManager::LoadModuleChecked<IWebBrowserAsPluginModule>( "WebBrowserAsPlugin" );
+	}
+
+	/**
+	 * Checks to see if this module is loaded and ready.  It is only valid to call Get() if IsAvailable() returns true.
+	 *
+	 * @return True if the module is loaded and ready to use
+	 */
+	static inline bool IsAvailable()
+	{
+		return FModuleManager::Get().IsModuleLoaded( "WebBrowserAsPlugin" );
+	}
+};
+
diff --git a/Engine/Plugins/Runtime/WebBrowserAsPlugin/Source/WebBrowserAsPlugin/WebBrowserAsPlugin.build.cs b/Engine/Plugins/Runtime/WebBrowserAsPlugin/Source/WebBrowserAsPlugin/WebBrowserAsPlugin.build.cs
new file mode 100644
--- /dev/null
+++ b/Engine/Plugins/Runtime/WebBrowserAsPlugin/Source/WebBrowserAsPlugin/WebBrowserAsPlugin.build.cs
@@ -0,0 +1,18 @@
+// Copyright 1998-2015 Epic Games, Inc. All Rights Reserved.
+
+namespace UnrealBuildTool.Rules
+{
+	public class WebBrowserAsPlugin : ModuleRules
+	{
+		public WebBrowserAsPlugin(TargetInfo Target)
+		{
+			PublicDependencyModuleNames.AddRange(
+				new string[]
+				{
+					"Core",
+					"WebBrowser"
+				}
+			);
+		}
+	}
+}
diff --git a/Engine/Plugins/Runtime/WebBrowserAsPlugin/WebBrowserAsPlugin.uplugin b/Engine/Plugins/Runtime/WebBrowserAsPlugin/WebBrowserAsPlugin.uplugin
new file mode 100644
--- /dev/null
+++ b/Engine/Plugins/Runtime/WebBrowserAsPlugin/WebBrowserAsPlugin.uplugin
@@ -0,0 +1,21 @@
+{
+    "PluginFileVersion" : 1,
+	
+	"FriendlyName" : "Web Browser",
+	"Version" : 1,
+	"VersionName" : "1.0",
+	"CreatedBy" : "Epic Games, Inc.",
+	"CreatedByURL" : "http://epicgames.com",
+	"EngineVersion" : "4.7.0",
+	"Description" : "Allows for selective use of the Web Browser module rather than building it into every game",
+	"CategoryPath" : "Programming Runtime.Plugins",
+	
+    "Modules" :
+	[
+		{
+			"Name" : "WebBrowserAsPlugin",
+			"Type" : "Runtime",
+			"WhitelistPlatforms" : [ "Win64", "Win32", "Mac" ]
+		}
+	]
+}
\ No newline at end of file
diff --git a/Engine/Source/UE4Game.Target.cs b/Engine/Source/UE4Game.Target.cs
--- a/Engine/Source/UE4Game.Target.cs
+++ b/Engine/Source/UE4Game.Target.cs
@@ -35,11 +35,6 @@ public class UE4GameTarget : TargetRules
 		{
 			OutExtraModuleNames.Add("GameMenuBuilder");
 
-			if (Target.Platform == UnrealTargetPlatform.Win32 || Target.Platform == UnrealTargetPlatform.Win64 || Target.Platform == UnrealTargetPlatform.Mac)
-			{
-				OutExtraModuleNames.Add("WebBrowser");
-			}
-
 			if (Target.Platform == UnrealTargetPlatform.Win32 || Target.Platform == UnrealTargetPlatform.Win64)
 			{
 				OutExtraModuleNames.Add("OnlineSubsystemNull");
