UMG - Adding a reference implementation of a web browser wrapper around the SWebBrowser Slate widget to the WebBrowserAsPlugin module.  Now if you enable the WebBrowserAsPlugin plugin for your project, you'll have a new web browser widget you can use.

diff --git a/Engine/Plugins/Runtime/WebBrowserAsPlugin/Source/WebBrowserAsPlugin/Private/WebBrowser.cpp b/Engine/Plugins/Runtime/WebBrowserAsPlugin/Source/WebBrowserAsPlugin/Private/WebBrowser.cpp
new file mode 100644
--- /dev/null
+++ b/Engine/Plugins/Runtime/WebBrowserAsPlugin/Source/WebBrowserAsPlugin/Private/WebBrowser.cpp
@@ -0,0 +1,98 @@
+// Copyright 1998-2015 Epic Games, Inc. All Rights Reserved.
+
+#include "WebBrowserAsPluginPrivatePCH.h"
+
+#include "SWebBrowser.h"
+
+#include "WebBrowser.h"
+
+#define LOCTEXT_NAMESPACE "WebBrowser"
+
+/////////////////////////////////////////////////////
+// UWebBrowser
+
+UWebBrowser::UWebBrowser(const FObjectInitializer& ObjectInitializer)
+	: Super(ObjectInitializer)
+{
+	bIsVariable = true;
+	InitialURL = TEXT("www.unrealengine.com");
+}
+
+void UWebBrowser::LoadURL(FString NewURL)
+{
+	if ( WebBrowserWidget.IsValid() )
+	{
+		return WebBrowserWidget->LoadURL(NewURL);
+	}
+}
+
+void UWebBrowser::LoadString(FString Contents, FString DummyURL)
+{
+	if ( WebBrowserWidget.IsValid() )
+	{
+		return WebBrowserWidget->LoadString(Contents, DummyURL);
+	}
+}
+
+FText UWebBrowser::GetTitleText() const
+{
+	if ( WebBrowserWidget.IsValid() )
+	{
+		return WebBrowserWidget->GetTitleText();
+	}
+
+	return FText::GetEmpty();
+}
+
+void UWebBrowser::ReleaseSlateResources(bool bReleaseChildren)
+{
+	Super::ReleaseSlateResources(bReleaseChildren);
+
+	WebBrowserWidget.Reset();
+}
+
+TSharedRef<SWidget> UWebBrowser::RebuildWidget()
+{
+	if ( IsDesignTime() )
+	{
+		return BuildDesignTimeWidget(SNew(SBox)
+			.HAlign(HAlign_Center)
+			.VAlign(VAlign_Center)
+			[
+				SNew(STextBlock)
+				.Text(LOCTEXT("Web Browser", "Web Browser"))
+			]);
+	}
+	else
+	{
+		WebBrowserWidget = SNew(SWebBrowser)
+			.InitialURL(InitialURL)
+			.ShowControls(false)
+			.SupportsTransparency(bSupportsTransparency);
+
+		return WebBrowserWidget.ToSharedRef();
+	}
+}
+
+void UWebBrowser::SynchronizeProperties()
+{
+	Super::SynchronizeProperties();
+
+	if ( WebBrowserWidget.IsValid() )
+	{
+
+	}
+}
+
+#if WITH_EDITOR
+
+const FText UWebBrowser::GetPaletteCategory()
+{
+	return LOCTEXT("Experimental", "Experimental");
+}
+
+#endif
+
+/////////////////////////////////////////////////////
+
+#undef LOCTEXT_NAMESPACE
diff --git a/Engine/Plugins/Runtime/WebBrowserAsPlugin/Source/WebBrowserAsPlugin/Public/WebBrowser.h b/Engine/Plugins/Runtime/WebBrowserAsPlugin/Source/WebBrowserAsPlugin/Public/WebBrowser.h
new file mode 100644
--- /dev/null
+++ b/Engine/Plugins/Runtime/WebBrowserAsPlugin/Source/WebBrowserAsPlugin/Public/WebBrowser.h
@@ -0,0 +1,70 @@
+// Copyright 1998-2015 Epic Games, Inc. All Rights Reserved.
+
+#pragma once
+
+#include "Widget.h"
+
+#include "WebBrowser.generated.h"
+
+//Experimental, ClassGroup=UserInterface
+
+/**
+ * 
+ */
+UCLASS()
+class WEBBROWSERASPLUGIN_API UWebBrowser : public UWidget
+{
+	GENERATED_UCLASS_BODY()
+
+public:
+
+	/**
+	 * Load the specified URL
+	 *
+	 * @param NewURL New URL to load
+	 */
+	UFUNCTION(BlueprintCallable, Category="Web Browser")
+	void LoadURL(FString NewURL);
+
+	/**
+	 * Load a string as data to create a web page
+	 *
+	 * @param Contents String to load
+	 * @param DummyURL Dummy URL for the page
+	 */
+	UFUNCTION(BlueprintCallable, Category="Web Browser")
+	void LoadString(FString Contents, FString DummyURL);
+
+	/**
+	 * Get the current title of the web page
+	 */
+	UFUNCTION(BlueprintCallable, Category="Web Browser")
+	FText GetTitleText() const;
+
+public:
+
+	// UWidget interface
+	virtual void SynchronizeProperties() override;
+	// End of UWidget interface
+
+	virtual void ReleaseSlateResources(bool bReleaseChildren) override;
+
+#if WITH_EDITOR
+	virtual const FText GetPaletteCategory() override;
+#endif
+
+protected:
+	UPROPERTY(EditDefaultsOnly, Category=Appearance)
+	FString InitialURL;
+
+	UPROPERTY(EditDefaultsOnly, Category=Appearance)
+	bool bSupportsTransparency;
+
+protected:
+	TSharedPtr<class SWebBrowser> WebBrowserWidget;
+
+protected:
+	// UWidget interface
+	virtual TSharedRef<SWidget> RebuildWidget() override;
+	// End of UWidget interface
+};
diff --git a/Engine/Plugins/Runtime/WebBrowserAsPlugin/Source/WebBrowserAsPlugin/WebBrowserAsPlugin.build.cs b/Engine/Plugins/Runtime/WebBrowserAsPlugin/Source/WebBrowserAsPlugin/WebBrowserAsPlugin.build.cs
--- a/Engine/Plugins/Runtime/WebBrowserAsPlugin/Source/WebBrowserAsPlugin/WebBrowserAsPlugin.build.cs
+++ b/Engine/Plugins/Runtime/WebBrowserAsPlugin/Source/WebBrowserAsPlugin/WebBrowserAsPlugin.build.cs
@@ -10,7 +10,11 @@ namespace UnrealBuildTool.Rules
 				new string[]
 				{
 					"Core",
-					"WebBrowser"
+				    "CoreUObject",
+					"WebBrowser",
+                    "Slate",
+                    "SlateCore",
+					"UMG"
 				}
 			);
 		}
