Fixed unity compile error

diff --git a/Engine/Plugins/Runtime/WebBrowserAsPlugin/Source/WebBrowserAsPlugin/Private/WebBrowserAsPluginPrivatePCH.h b/Engine/Plugins/Runtime/WebBrowserAsPlugin/Source/WebBrowserAsPlugin/Private/WebBrowserAsPluginPrivatePCH.h
--- a/Engine/Plugins/Runtime/WebBrowserAsPlugin/Source/WebBrowserAsPlugin/Private/WebBrowserAsPluginPrivatePCH.h
+++ b/Engine/Plugins/Runtime/WebBrowserAsPlugin/Source/WebBrowserAsPlugin/Private/WebBrowserAsPluginPrivatePCH.h
@@ -1,5 +1,7 @@
 // Copyright 1998-2015 Epic Games, Inc. All Rights Reserved.
 
+#pragma once
+
 #include "WebBrowserAsPluginModule.h"
 
 DECLARE_LOG_CATEGORY_EXTERN(LogWebBrowserAsPlugin, Verbose, All);
\ No newline at end of file
