Fix Mac compile issues

diff --git a/Engine/Plugins/Runtime/WebBrowserAsPlugin/Source/WebBrowserAsPlugin/Private/WebBrowserAsPluginPrivatePCH.h b/Engine/Plugins/Runtime/WebBrowserAsPlugin/Source/WebBrowserAsPlugin/Private/WebBrowserAsPluginPrivatePCH.h
--- a/Engine/Plugins/Runtime/WebBrowserAsPlugin/Source/WebBrowserAsPlugin/Private/WebBrowserAsPluginPrivatePCH.h
+++ b/Engine/Plugins/Runtime/WebBrowserAsPlugin/Source/WebBrowserAsPlugin/Private/WebBrowserAsPluginPrivatePCH.h
@@ -2,6 +2,9 @@
 
 #pragma once
 
+#include "SlateBasics.h"
+#include "SlateCore.h"
+
 #include "WebBrowserAsPluginModule.h"
 
 DECLARE_LOG_CATEGORY_EXTERN(LogWebBrowserAsPlugin, Verbose, All);
\ No newline at end of file
