UMG - Cleaning up some unneeded comments in WebBrowser.h

diff --git a/Engine/Plugins/Runtime/WebBrowserAsPlugin/Source/WebBrowserAsPlugin/Public/WebBrowser.h b/Engine/Plugins/Runtime/WebBrowserAsPlugin/Source/WebBrowserAsPlugin/Public/WebBrowser.h
--- a/Engine/Plugins/Runtime/WebBrowserAsPlugin/Source/WebBrowserAsPlugin/Public/WebBrowser.h
+++ b/Engine/Plugins/Runtime/WebBrowserAsPlugin/Source/WebBrowserAsPlugin/Public/WebBrowser.h
@@ -6,8 +6,6 @@
 
 #include "WebBrowser.generated.h"
 
-//Experimental, ClassGroup=UserInterface
-
 /**
  * 
  */
